Fixing build

diff --git a/Engine/Source/Editor/CookingStats/Public/CookingStatsModule.h b/Engine/Source/Editor/CookingStats/Public/CookingStatsModule.h
--- a/Engine/Source/Editor/CookingStats/Public/CookingStatsModule.h
+++ b/Engine/Source/Editor/CookingStats/Public/CookingStatsModule.h
@@ -5,7 +5,7 @@
 
 #include "ModuleInterface.h"
 #include "ICookingStats.h"
-#include "CookingStatsInterface.h"
+//#include "CookingStatsInterface.h"
 
 namespace CookingStatsConstants
 {
@@ -15,7 +15,9 @@ namespace CookingStatsConstants
 /**
 * Asset registry module
 */
-class FCookingStatsModule : public ICookingStatsInterface
+class FCookingStatsModule
+	: public IModuleInterface
+	//: public ICookingStatsModule
 {
 
 public:
